PluginsEditor: Code cleanup pass

- local callback to lambda
- NULL to nullptr

diff --git a/Engine/Plugins/Editor/PluginsEditor/Source/PluginsEditor/Private/SPluginListTile.cpp b/Engine/Plugins/Editor/PluginsEditor/Source/PluginsEditor/Private/SPluginListTile.cpp
--- a/Engine/Plugins/Editor/PluginsEditor/Source/PluginsEditor/Private/SPluginListTile.cpp
+++ b/Engine/Plugins/Editor/PluginsEditor/Source/PluginsEditor/Private/SPluginListTile.cpp
@@ -46,20 +46,13 @@ void SPluginListTile::Construct( const FArguments& Args, const TSharedRef< SPlug
 	TSharedPtr<SWidget> CreatedByWidget;
 	if( !Item->PluginStatus.CreatedBy.IsEmpty() && !Item->PluginStatus.CreatedByURL.IsEmpty() )
 	{
-		struct Local
-		{
-			static void NavigateToURL( FString URL )
-			{
-				FPlatformProcess::LaunchURL( *URL, NULL, NULL );
-			}
-		};
+		FString CreatedByURL = Item->PluginStatus.CreatedByURL;
 
 		// Clickable 'Created by' URL text
 		CreatedByWidget = SNew( SHyperlink )
 			.Text( Item->PluginStatus.CreatedBy )
-			.ToolTipText( FText::Format( LOCTEXT("NavigateToCreatedByURL", "Launch a web browser to visit {0}"), FText::FromString( Item->PluginStatus.CreatedByURL ) ) )
-			.OnNavigate_Static( &Local::NavigateToURL, Item->PluginStatus.CreatedByURL );
-
+			.ToolTipText( FText::Format( LOCTEXT("NavigateToCreatedByURL", "Launch a web browser to visit {0}"), FText::FromString( CreatedByURL ) ) )
+			.OnNavigate_Lambda([=]() { FPlatformProcess::LaunchURL( *CreatedByURL, nullptr, nullptr ); });
 	}
 	else if( !Item->PluginStatus.CreatedBy.IsEmpty() )
 	{
@@ -97,7 +90,7 @@ void SPluginListTile::Construct( const FArguments& Args, const TSharedRef< SPlug
 					.HeightOverride( ThumbnailImageSize )
 					[
 						SNew( SImage )
-						.Image( PluginIconDynamicImageBrush.IsValid() ? PluginIconDynamicImageBrush.Get() : NULL )
+						.Image( PluginIconDynamicImageBrush.IsValid() ? PluginIconDynamicImageBrush.Get() : nullptr )
 					]
 				]
 
