PluginEditor: Burst upgrades; added documentation links

diff --git a/Engine/Plugins/Editor/PluginsEditor/Source/PluginsEditor/Private/SPluginListTile.cpp b/Engine/Plugins/Editor/PluginsEditor/Source/PluginsEditor/Private/SPluginListTile.cpp
--- a/Engine/Plugins/Editor/PluginsEditor/Source/PluginsEditor/Private/SPluginListTile.cpp
+++ b/Engine/Plugins/Editor/PluginsEditor/Source/PluginsEditor/Private/SPluginListTile.cpp
@@ -11,7 +11,7 @@
 #define LOCTEXT_NAMESPACE "PluginListTile"
 
 
-void SPluginListTile::Construct( const FArguments& Args, const TSharedRef< SPluginList > Owner, const TSharedRef< class FPluginListItem >& Item )
+void SPluginListTile::Construct( const FArguments& Args, const TSharedRef<SPluginList> Owner, const TSharedRef<class FPluginListItem>& Item )
 {
 	OwnerWeak = Owner;
 	ItemData = Item;
@@ -23,15 +23,15 @@ void SPluginListTile::Construct( const FArguments& Args, const TSharedRef< SPlug
 	//		-> Maybe a filter for this too?  (show only editor plugins, etc.)
 	// @todo plugedit: Indicate whether plugin has content?  Filter to show only content plugins, and vice-versa?
 
-
 	// Plugin thumbnail image
 	if( !Item->PluginStatus.Icon128FilePath.IsEmpty() )
 	{
 		const FName BrushName( *Item->PluginStatus.Icon128FilePath );
-		const FIntPoint Size = FSlateApplication::Get().GetRenderer()->GenerateDynamicImageResource( BrushName );
-		if( Size.X > 0 && Size.Y > 0 )
+		const FIntPoint Size = FSlateApplication::Get().GetRenderer()->GenerateDynamicImageResource(BrushName);
+
+		if ((Size.X > 0) && (Size.Y > 0))
 		{
-			PluginIconDynamicImageBrush = MakeShareable( new FSlateDynamicImageBrush( BrushName, FVector2D( Size.X, Size.Y ) ) );
+			PluginIconDynamicImageBrush = MakeShareable(new FSlateDynamicImageBrush(BrushName, FVector2D(Size.X, Size.Y)));
 		}
 	}
 	else
@@ -40,215 +40,271 @@ void SPluginListTile::Construct( const FArguments& Args, const TSharedRef< SPlug
 		// @todo plugedit: Should display default image or just omit the thumbnail
 	}
 
-	
-	// Setup widget for the 'Created by' text.  Depending on how the user configured their plugin, we may display a plain text label, a hyperlink to
-	// the company's web site, or nothing at all!
-	TSharedPtr<SWidget> CreatedByWidget;
-	if( !Item->PluginStatus.CreatedBy.IsEmpty() && !Item->PluginStatus.CreatedByURL.IsEmpty() )
+	// create documentation link
+	TSharedPtr<SWidget> DocumentationWidget;
 	{
-		FString CreatedByURL = Item->PluginStatus.CreatedByURL;
+		if (Item->PluginStatus.DocsURL.IsEmpty())
+		{
+			DocumentationWidget = SNullWidget::NullWidget;
+		}
+		else
+		{
+			FString DocsURL = Item->PluginStatus.DocsURL;
+			DocumentationWidget = SNew(SHorizontalBox)
 
-		// Clickable 'Created by' URL text
-		CreatedByWidget = SNew( SHyperlink )
-			.Text( Item->PluginStatus.CreatedBy )
-			.ToolTipText( FText::Format( LOCTEXT("NavigateToCreatedByURL", "Launch a web browser to visit {0}"), FText::FromString( CreatedByURL ) ) )
-			.OnNavigate_Lambda([=]() { FPlatformProcess::LaunchURL( *CreatedByURL, nullptr, nullptr ); });
-	}
-	else if( !Item->PluginStatus.CreatedBy.IsEmpty() )
-	{
-		// Plain text for 'Created by' (no URL specified)
-		CreatedByWidget = 
-			SNew( STextBlock )
-			.Text( Item->PluginStatus.CreatedBy );
+			+ SHorizontalBox::Slot()
+				.AutoWidth()
+				.VAlign(VAlign_Center)
+				[
+					SNew(SImage)
+						.ColorAndOpacity(FSlateColor::UseForeground())
+						.Image(FEditorStyle::GetBrush("MessageLog.Docs"))
+				]
+
+			+ SHorizontalBox::Slot()
+				.AutoWidth()
+				.VAlign(VAlign_Center)
+				.Padding(2.0f, 0.0f, 0.0f, 0.0f)
+				[
+					SNew(SHyperlink)
+						.Text(LOCTEXT("DocumentationLink", "Documentation"))
+						.ToolTipText(LOCTEXT("NavigateToCreatedByURL", "Open the plug-in's online documentation"))
+						.OnNavigate_Lambda([=]() { FPlatformProcess::LaunchURL(*DocsURL, nullptr, nullptr); })
+				];
+		}
 	}
-	else
+
+	// create vendor link
+	TSharedPtr<SWidget> CreatedByWidget;
 	{
-		// No 'Created by' text was supplied
-		CreatedByWidget = SNew( SSpacer );
-	}
+		if (Item->PluginStatus.CreatedBy.IsEmpty())
+		{
+			CreatedByWidget = SNullWidget::NullWidget;
+		}
+		else if (Item->PluginStatus.CreatedByURL.IsEmpty())
+		{
+			CreatedByWidget = SNew(SHorizontalBox)
 
+			+ SHorizontalBox::Slot()
+				.AutoWidth()
+				.VAlign(VAlign_Center)
+				[
+					SNew(SImage)
+						.ColorAndOpacity(FSlateColor::UseForeground())
+						.Image(FEditorStyle::GetBrush("ContentBrowser.AssetTreeFolderDeveloper"))
+				]
 
-	ChildSlot
-	[
-		SNew( SBorder )
-		.BorderImage( FEditorStyle::GetBrush( "NoBorder" ) )
-		.Padding( PaddingAmount )
-		[
-			SNew( SBorder )
-			.BorderImage( FEditorStyle::GetBrush( "ToolPanel.GroupBorder" ) )
-			.Padding( PaddingAmount )
-			[
-				SNew( SHorizontalBox )
+			+ SHorizontalBox::Slot()
+				.AutoWidth()
+				.VAlign(VAlign_Center)
+				.Padding(2.0f, 0.0f, 0.0f, 0.0f)
+				[
+					SNew(STextBlock)
+						.Text(Item->PluginStatus.CreatedBy)
+				];
+		}
+		else
+		{
+			FString CreatedByURL = Item->PluginStatus.CreatedByURL;
+			CreatedByWidget = SNew(SHorizontalBox)
 
-				// Thumbnail image
-				+SHorizontalBox::Slot()
-				.Padding( PaddingAmount )
+			+ SHorizontalBox::Slot()
 				.AutoWidth()
+				.VAlign(VAlign_Center)
 				[
-					SNew( SBox )
-					.WidthOverride( ThumbnailImageSize )
-					.HeightOverride( ThumbnailImageSize )
-					[
-						SNew( SImage )
-						.Image( PluginIconDynamicImageBrush.IsValid() ? PluginIconDynamicImageBrush.Get() : nullptr )
-					]
+					SNew(SImage)
+						.ColorAndOpacity(FSlateColor::UseForeground())
+						.Image(FEditorStyle::GetBrush("MessageLog.Url"))
 				]
 
-				+SHorizontalBox::Slot()
-				[
-					SNew( SVerticalBox )
+			+ SHorizontalBox::Slot()
+				.AutoWidth()
+				.VAlign(VAlign_Center)
+				.Padding(2.0f, 0.0f, 0.0f, 0.0f)
+				[				
+					SNew(SHyperlink)
+						.Text(Item->PluginStatus.CreatedBy)
+						.ToolTipText(FText::Format(LOCTEXT("NavigateToCreatedByURL", "Visit the vendor's web site ({0})"), FText::FromString(CreatedByURL)))
+						.OnNavigate_Lambda([=]() { FPlatformProcess::LaunchURL(*CreatedByURL, nullptr, nullptr); })
+				];
+		}
+	}
 
-					+SVerticalBox::Slot()
-					.AutoHeight()
+	ChildSlot
+	[
+		SNew(SBorder)
+			.BorderImage(FEditorStyle::GetBrush("NoBorder"))
+			.Padding(PaddingAmount)
+			[
+				SNew(SBorder)
+					.BorderImage(FEditorStyle::GetBrush("ToolPanel.GroupBorder"))
+					.Padding(PaddingAmount)
 					[
-						SNew( SHorizontalBox )
+						SNew(SHorizontalBox)
 
-						// Friendly name
-						+SHorizontalBox::Slot()
-						.Padding( PaddingAmount )
-						[
-							SNew( STextBlock )
-							.Text( Item->PluginStatus.FriendlyName )
-							.HighlightText_Raw( &Owner->GetOwner().GetPluginTextFilter(), &FPluginTextFilter::GetRawFilterText )
-							.TextStyle( FPluginStyle::Get(), "PluginTile.NameText" )
-						]
-
-						// Version
-						+SHorizontalBox::Slot()
-						.HAlign( HAlign_Right )
-						.Padding( PaddingAmount )
-						.AutoWidth()
-						[
-							SNew( SHorizontalBox )
-							+SHorizontalBox::Slot()
+						// Thumbnail image
+						+ SHorizontalBox::Slot()
+							.Padding(PaddingAmount)
 							.AutoWidth()
-							.Padding( 0.0f, 0.0f, 0.0f, 1.0f )		// Lower padding to align font with version number base
 							[
-								SNew( SHorizontalBox )
-
-								// Beta version icon
-								+SHorizontalBox::Slot()
-								.AutoWidth()
-								.VAlign( VAlign_Bottom )
-								.Padding( 0.0f, 0.0f, 3.0f, 2.0f )
-								[
-									SNew( SImage )
-									.Visibility( Item->PluginStatus.bIsBetaVersion ? EVisibility::Visible : EVisibility::Collapsed )
-									.Image( FPluginStyle::Get()->GetBrush("PluginTile.BetaWarning") )
-								]
-								// Version label
-								+SHorizontalBox::Slot()
-								.AutoWidth()
-								.VAlign( VAlign_Bottom )
-								[
-									SNew( STextBlock )
-									.Text( Item->PluginStatus.bIsBetaVersion ?
-										LOCTEXT( "PluginBetaVersionLabel", "BETA Version " ) :
-										LOCTEXT( "PluginVersionLabel", "Version " ) )
-								]
+								SNew(SBox)
+									.WidthOverride(ThumbnailImageSize)
+									.HeightOverride(ThumbnailImageSize)
+									[
+										SNew(SImage)
+											.Image(PluginIconDynamicImageBrush.IsValid() ? PluginIconDynamicImageBrush.Get() : nullptr)
+									]
 							]
 
-							+SHorizontalBox::Slot()
-							.AutoWidth()
-							.VAlign( VAlign_Bottom )
-							.Padding( 0.0f, 0.0f, 2.0f, 0.0f )	// Extra padding from the right edge
+						+SHorizontalBox::Slot()
 							[
-								SNew( STextBlock )
-								.Text( Item->PluginStatus.VersionName )
-								.TextStyle( FPluginStyle::Get(), "PluginTile.VersionNumberText" )
-							]
-						]
-					]
+								SNew(SVerticalBox)
+
+								+ SVerticalBox::Slot()
+									.AutoHeight()
+									[
+										SNew(SHorizontalBox)
+
+										// Friendly name
+										+SHorizontalBox::Slot()
+										.Padding(PaddingAmount)
+										[
+											SNew(STextBlock)
+												.Text(Item->PluginStatus.FriendlyName)
+												.HighlightText_Raw(&Owner->GetOwner().GetPluginTextFilter(), &FPluginTextFilter::GetRawFilterText)
+												.TextStyle(FPluginStyle::Get(), "PluginTile.NameText")
+										]
+
+										// Version
+										+ SHorizontalBox::Slot()
+											.HAlign(HAlign_Right)
+											.Padding(PaddingAmount)
+											.AutoWidth()
+											[
+												SNew(SHorizontalBox)
+
+												+ SHorizontalBox::Slot()
+													.AutoWidth()
+													.Padding(0.0f, 0.0f, 0.0f, 1.0f) // Lower padding to align font with version number base
+													[
+														SNew(SHorizontalBox)
+
+														// beta version icon
+														+ SHorizontalBox::Slot()
+															.AutoWidth()
+															.VAlign(VAlign_Bottom)
+															.Padding(0.0f, 0.0f, 3.0f, 2.0f)
+															[
+																SNew(SImage)
+																	.Visibility(Item->PluginStatus.bIsBetaVersion ? EVisibility::Visible : EVisibility::Collapsed)
+																	.Image(FPluginStyle::Get()->GetBrush("PluginTile.BetaWarning"))
+															]
+
+														// version label
+														+ SHorizontalBox::Slot()
+															.AutoWidth()
+															.VAlign(VAlign_Bottom)
+															[
+																SNew(STextBlock)
+																	.Text(Item->PluginStatus.bIsBetaVersion ? LOCTEXT("PluginBetaVersionLabel", "BETA Version ") : LOCTEXT("PluginVersionLabel", "Version "))
+															]
+													]
+
+												+ SHorizontalBox::Slot()
+													.AutoWidth()
+													.VAlign( VAlign_Bottom )
+													.Padding( 0.0f, 0.0f, 2.0f, 0.0f )	// Extra padding from the right edge
+													[
+														SNew(STextBlock)
+															.Text(Item->PluginStatus.VersionName)
+															.TextStyle(FPluginStyle::Get(), "PluginTile.VersionNumberText")
+													]
+											]
+									]
 			
-					+SVerticalBox::Slot()
-					[
-						SNew( SVerticalBox )
+								+ SVerticalBox::Slot()
+									[
+										SNew(SVerticalBox)
 				
-						// Description
-						+SVerticalBox::Slot()
-						.Padding( PaddingAmount )
-						[
-							SNew( STextBlock )
-							.Text( Item->PluginStatus.Description )
-							.AutoWrapText( true )
-						]
-
-						+SVerticalBox::Slot()
-						.Padding( PaddingAmount )
-						.AutoHeight()
-						[
-							SNew( SHorizontalBox )
-
-							// Enable checkbox
-							+SHorizontalBox::Slot()
-							.Padding( PaddingAmount )
-							.AutoWidth()
-							[
-								SNew(SCheckBox)
-								.OnCheckStateChanged(this, &SPluginListTile::OnEnablePluginCheckboxChanged)
-								.IsChecked( this, &SPluginListTile::IsPluginEnabled )
-								.ToolTipText( LOCTEXT( "EnableDisableButtonToolTip", "Toggles whether this plugin is enabled for your current project.  You may need to restart the program for this change to take effect." ) )
-								.Content()
-								[
-									SNew(STextBlock) .Text( LOCTEXT("EnablePluginCheckbox", "Enabled") )
-								]
-							]
-
-							// Uninstall button
-							+SHorizontalBox::Slot()
-							.Padding( PaddingAmount )
-							.AutoWidth()
-							[
-								SNew( SButton )
-								.Text( LOCTEXT( "PluginUninstallButtonLabel", "Uninstall..." ) )
-
-								// Don't show 'Uninstall' for built-in plugins
-								.Visibility( EVisibility::Collapsed )	
-								// @todo plugins: Not supported yet
-									// .Visibility( Item->PluginStatus.bIsBuiltIn ? EVisibility::Collapsed : EVisibility::Visible )
-
-								.OnClicked( this, &SPluginListTile::OnUninstallButtonClicked )
-								.IsEnabled( false )		// @todo plugedit: Not supported yet
-								.ToolTipText( LOCTEXT( "UninstallButtonToolTip", "Allows you to uninstall this plugin.  This will delete the plugin from your project entirely.  You may need to restart the program for this change to take effect." ) )
+										// Description
+										+ SVerticalBox::Slot()
+											.Padding( PaddingAmount )
+											[
+												SNew(STextBlock)
+													.Text(Item->PluginStatus.Description)
+													.AutoWrapText(true)
+											]
+
+										+ SVerticalBox::Slot()
+											.Padding(PaddingAmount)
+											.AutoHeight()
+											[
+												SNew(SHorizontalBox)
+
+												// Enable checkbox
+												+ SHorizontalBox::Slot()
+													.Padding(PaddingAmount)
+													.AutoWidth()
+													[
+														SNew(SCheckBox)
+															.OnCheckStateChanged(this, &SPluginListTile::OnEnablePluginCheckboxChanged)
+															.IsChecked(this, &SPluginListTile::IsPluginEnabled)
+															.ToolTipText(LOCTEXT("EnableDisableButtonToolTip", "Toggles whether this plugin is enabled for your current project.  You may need to restart the program for this change to take effect."))
+															.Content()
+															[
+																SNew(STextBlock)
+																	.Text(LOCTEXT("EnablePluginCheckbox", "Enabled"))
+															]
+													]
+
+												// Uninstall button
+												+ SHorizontalBox::Slot()
+													.Padding(PaddingAmount)
+													.AutoWidth()
+													[
+														SNew(SButton)
+															.Text(LOCTEXT("PluginUninstallButtonLabel", "Uninstall..."))
+
+															// Don't show 'Uninstall' for built-in plugins
+															.Visibility(EVisibility::Collapsed)
+															// @todo plugins: Not supported yet
+																// .Visibility( Item->PluginStatus.bIsBuiltIn ? EVisibility::Collapsed : EVisibility::Visible )
+
+															.OnClicked(this, &SPluginListTile::OnUninstallButtonClicked)
+															.IsEnabled(false)		// @todo plugedit: Not supported yet
+															.ToolTipText(LOCTEXT("UninstallButtonToolTip", "Allows you to uninstall this plugin.  This will delete the plugin from your project entirely.  You may need to restart the program for this change to take effect."))
+													]
+
+												// docs link
+												+ SHorizontalBox::Slot()
+													.Padding(PaddingAmount)
+													.HAlign(HAlign_Right)
+													[
+														DocumentationWidget.ToSharedRef()
+													]
+
+												// vendor link
+												+ SHorizontalBox::Slot()
+													.AutoWidth()
+													.Padding(12.0f, PaddingAmount, PaddingAmount, PaddingAmount)
+													.HAlign(HAlign_Right)
+													[
+														CreatedByWidget.ToSharedRef()
+													]
+											]
+									]
 							]
-
-							// Created by
-							+SHorizontalBox::Slot()
-							.Padding( PaddingAmount )
-							.HAlign( HAlign_Right )
-							[
-								SNew( SHorizontalBox )
-
-								+SHorizontalBox::Slot()
-								.AutoWidth()
-								.VAlign( VAlign_Bottom )
-								.Padding( 0, 0, 4, 0 )		// Extra padding before the URL string
-								[
-									SNew( STextBlock )
-									.Visibility( Item->PluginStatus.CreatedBy.IsEmpty() ? EVisibility::Collapsed : EVisibility::Visible )
-									.Text( LOCTEXT( "PluginCreatedByLabel", "Created by" ) )
-									.Text( Item->PluginStatus.CreatedBy )							
-								]
-
-								+SHorizontalBox::Slot()
-								.VAlign( VAlign_Bottom )
-								.AutoWidth()
-								[
-									CreatedByWidget.ToSharedRef()
-								]
-							]
-						]
 					]
-				]
 			]
-		]
 	];
 }
 
 
 ESlateCheckBoxState::Type SPluginListTile::IsPluginEnabled() const
 {
-	return ItemData->PluginStatus.bIsEnabled ? ESlateCheckBoxState::Checked : ESlateCheckBoxState::Unchecked;
+	return ItemData->PluginStatus.bIsEnabled
+		? ESlateCheckBoxState::Checked
+		: ESlateCheckBoxState::Unchecked;
 }
 
 
@@ -264,17 +320,16 @@ void SPluginListTile::OnEnablePluginCheckboxChanged(ESlateCheckBoxState::Type Ne
 
 		if (EAppReturnType::No == FMessageDialog::Open(EAppMsgType::YesNo, WarningMessage))
 		{
-			// User chose to keep beta plugin disabled
+			// user chose to keep beta plug-in disabled
 			return;
 		}
 	}
 
 	ItemData->PluginStatus.bIsEnabled = bNewEnabledState;
-
 	FGameProjectGenerationModule::Get().TryMakeProjectFileWriteable(FPaths::GetProjectFilePath());
-
 	FText FailMessage;
-	if(!IProjectManager::Get().SetPluginEnabled(ItemData->PluginStatus.Name, bNewEnabledState, FailMessage))
+
+	if (!IProjectManager::Get().SetPluginEnabled(ItemData->PluginStatus.Name, bNewEnabledState, FailMessage))
 	{
 		FMessageDialog::Open(EAppMsgType::Ok, FailMessage);
 	}
@@ -283,8 +338,8 @@ void SPluginListTile::OnEnablePluginCheckboxChanged(ESlateCheckBoxState::Type Ne
 
 FReply SPluginListTile::OnUninstallButtonClicked()
 {
-	// Not expecting to deal with engine plugins here.  They can't be unintalled
-	if( ensure( !ItemData->PluginStatus.bIsBuiltIn ) )
+	// not expecting to deal with engine plugins here (they can't be uninstalled)
+	if (ensure(!ItemData->PluginStatus.bIsBuiltIn))
 	{
 		// @todo plugedit: Make this work!
 	}
@@ -293,4 +348,4 @@ FReply SPluginListTile::OnUninstallButtonClicked()
 }
 
 
-#undef LOCTEXT_NAMESPACE
\ No newline at end of file
+#undef LOCTEXT_NAMESPACE
diff --git a/Engine/Plugins/Media/AvfMedia/AvfMedia.uplugin b/Engine/Plugins/Media/AvfMedia/AvfMedia.uplugin
--- a/Engine/Plugins/Media/AvfMedia/AvfMedia.uplugin
+++ b/Engine/Plugins/Media/AvfMedia/AvfMedia.uplugin
@@ -6,6 +6,7 @@
 	"VersionName" : "1.0",
 	"CreatedBy" : "Epic Games, Inc.",
 	"CreatedByURL" : "http://epicgames.com",
+	"DocsURL" : "https://forums.unrealengine.com/showthread.php?46879-Media-Framework-Documentation-for-4-5-Preview",
 	"EngineVersion" : "4.4.0",
 	"Description" : "Implements a media player using the AV Foundation framework.",
 	"Category" : "Media.Players",
diff --git a/Engine/Plugins/Media/MediaPlayerEditor/MediaPlayerEditor.uplugin b/Engine/Plugins/Media/MediaPlayerEditor/MediaPlayerEditor.uplugin
--- a/Engine/Plugins/Media/MediaPlayerEditor/MediaPlayerEditor.uplugin
+++ b/Engine/Plugins/Media/MediaPlayerEditor/MediaPlayerEditor.uplugin
@@ -6,6 +6,7 @@
 	"VersionName" : "1.0",
 	"CreatedBy" : "Epic Games, Inc.",
 	"CreatedByURL" : "http://epicgames.com",
+	"DocsURL" : "https://forums.unrealengine.com/showthread.php?46879-Media-Framework-Documentation-for-4-5-Preview",
 	"EngineVersion" : "4.5.0",
 	"Description" : "Content Editor for MediaPlayer Assets.",
 	"Category" : "Media.Editors",
diff --git a/Engine/Plugins/Media/WmfMedia/WmfMedia.uplugin b/Engine/Plugins/Media/WmfMedia/WmfMedia.uplugin
--- a/Engine/Plugins/Media/WmfMedia/WmfMedia.uplugin
+++ b/Engine/Plugins/Media/WmfMedia/WmfMedia.uplugin
@@ -6,6 +6,7 @@
 	"VersionName" : "1.0",
 	"CreatedBy" : "Epic Games, Inc.",
 	"CreatedByURL" : "http://epicgames.com",
+	"DocsURL" : "https://forums.unrealengine.com/showthread.php?46879-Media-Framework-Documentation-for-4-5-Preview",
 	"EngineVersion" : "4.2.0",
 	"Description" : "Implements a media player using the Windows Media Foundation framework.",
 	"Category" : "Media.Players",
diff --git a/Engine/Source/Runtime/Projects/Private/PluginDescriptor.cpp b/Engine/Source/Runtime/Projects/Private/PluginDescriptor.cpp
--- a/Engine/Source/Runtime/Projects/Private/PluginDescriptor.cpp
+++ b/Engine/Source/Runtime/Projects/Private/PluginDescriptor.cpp
@@ -5,42 +5,47 @@
 
 #define LOCTEXT_NAMESPACE "PluginDescriptor"
 
+
 FPluginDescriptor::FPluginDescriptor()
-{
-	FileVersion = EPluginDescriptorVersion::Latest;
-	Version = 0;
-	bEnabledByDefault = false;
-	bCanContainContent = false;
-	bIsBetaVersion = false;
-}
+	: FileVersion(EPluginDescriptorVersion::Latest)
+	, Version(0)
+	, bEnabledByDefault(false)
+	, bCanContainContent(false)
+	, bIsBetaVersion(false)
+{ }
 
-bool FPluginDescriptor::Load(const FString& FileName, FText& OutFailReason)
+
+bool FPluginDescriptor::Load( const FString& FileName, FText& OutFailReason )
 {
 	// Read the file to a string
 	FString FileContents;
+
 	if (!FFileHelper::LoadFileToString(FileContents, *FileName))
 	{
-		OutFailReason = FText::Format( LOCTEXT("FailedToLoadDescriptorFile", "Failed to open descriptor file '{0}'"), FText::FromString( FileName ) );
+		OutFailReason = FText::Format(LOCTEXT("FailedToLoadDescriptorFile", "Failed to open descriptor file '{0}'"), FText::FromString(FileName));
 		return false;
 	}
 
 	// Deserialize a JSON object from the string
 	TSharedPtr< FJsonObject > Object;
 	TSharedRef< TJsonReader<> > Reader = TJsonReaderFactory<>::Create(FileContents);
-	if ( !FJsonSerializer::Deserialize(Reader, Object) || !Object.IsValid() )
+	
+	if (!FJsonSerializer::Deserialize(Reader, Object) || !Object.IsValid() )
 	{
-		OutFailReason = FText::Format( LOCTEXT("FailedToReadDescriptorFile", "Failed to read file. {0}"), FText::FromString(Reader->GetErrorMessage()) );
+		OutFailReason = FText::Format(LOCTEXT("FailedToReadDescriptorFile", "Failed to read file. {0}"), FText::FromString(Reader->GetErrorMessage()));
 		return false;
 	}
 
-	// Parse it as a plugin descriptor
+	// Parse it as a plug-in descriptor
 	return Read(*Object.Get(), OutFailReason);
 }
 
-bool FPluginDescriptor::Read(const FJsonObject& Object, FText& OutFailReason)
+
+bool FPluginDescriptor::Read( const FJsonObject& Object, FText& OutFailReason )
 {
 	// Read the file version
 	int32 FileVersionInt32;
+
 	if(!Object.TryGetNumberField(TEXT("FileVersion"), FileVersionInt32))
 	{
 		if(!Object.TryGetNumberField(TEXT("PluginFileVersion"), FileVersionInt32))
@@ -52,10 +57,10 @@ bool FPluginDescriptor::Read(const FJsonObject& Object, FText& OutFailReason)
 
 	// Check that it's within range
 	EPluginDescriptorVersion::Type FileVersion = (EPluginDescriptorVersion::Type)FileVersionInt32;
-	if ( FileVersion <= EPluginDescriptorVersion::Invalid || FileVersion > EPluginDescriptorVersion::Latest )
+	if ((FileVersion <= EPluginDescriptorVersion::Invalid) || (FileVersion > EPluginDescriptorVersion::Latest))
 	{
-		FText ReadVersionText = FText::FromString( FString::Printf( TEXT( "%d" ), (int32)FileVersion ) );
-		FText LatestVersionText = FText::FromString( FString::Printf( TEXT( "%d" ), (int32)EPluginDescriptorVersion::Latest ) );
+		FText ReadVersionText = FText::FromString(FString::Printf(TEXT("%d"), (int32)FileVersion));
+		FText LatestVersionText = FText::FromString(FString::Printf(TEXT("%d"), (int32)EPluginDescriptorVersion::Latest));
 		OutFailReason = FText::Format( LOCTEXT("ProjectFileVersionTooLarge", "File appears to be in a newer version ({0}) of the file format that we can load (max version: {1})."), ReadVersionText, LatestVersionText);
 		return false;
 	}
@@ -66,7 +71,7 @@ bool FPluginDescriptor::Read(const FJsonObject& Object, FText& OutFailReason)
 	Object.TryGetStringField(TEXT("FriendlyName"), FriendlyName);
 	Object.TryGetStringField(TEXT("Description"), Description);
 
-	if ( !Object.TryGetStringField(TEXT("Category"), Category) )
+	if (!Object.TryGetStringField(TEXT("Category"), Category))
 	{
 		// Category used to be called CategoryPath in .uplugin files
 		Object.TryGetStringField(TEXT("CategoryPath"), Category);
@@ -74,15 +79,16 @@ bool FPluginDescriptor::Read(const FJsonObject& Object, FText& OutFailReason)
         
 	// Due to a difference in command line parsing between Windows and Mac, we shipped a few Mac samples containing
 	// a category name with escaped quotes. Remove them here to make sure we can list them in the right category.
-	if(Category.Len() >= 2 && Category.StartsWith(TEXT("\"")) && Category.EndsWith(TEXT("\"")))
+	if (Category.Len() >= 2 && Category.StartsWith(TEXT("\"")) && Category.EndsWith(TEXT("\"")))
 	{
 		Category = Category.Mid(1, Category.Len() - 2);
 	}
 
 	Object.TryGetStringField(TEXT("CreatedBy"), CreatedBy);
 	Object.TryGetStringField(TEXT("CreatedByURL"), CreatedByURL);
+	Object.TryGetStringField(TEXT("DocsURL"), DocsURL);
 
-	if(!FModuleDescriptor::ReadArray(Object, TEXT("Modules"), Modules, OutFailReason))
+	if (!FModuleDescriptor::ReadArray(Object, TEXT("Modules"), Modules, OutFailReason))
 	{
 		return false;
 	}
@@ -90,23 +96,18 @@ bool FPluginDescriptor::Read(const FJsonObject& Object, FText& OutFailReason)
 	Object.TryGetBoolField(TEXT("EnabledByDefault"), bEnabledByDefault);
 	Object.TryGetBoolField(TEXT("CanContainContent"), bCanContainContent);
 	Object.TryGetBoolField(TEXT("IsBetaVersion"), bIsBetaVersion);
+
 	return true;
 }
 
 
-
-
-
-
-
-
-FPluginReferenceDescriptor::FPluginReferenceDescriptor(const FString &InName, bool bInEnabled)
+FPluginReferenceDescriptor::FPluginReferenceDescriptor( const FString &InName, bool bInEnabled )
 	: Name(InName)
 	, bEnabled(bInEnabled)
-{
-}
+{ }
+
 
-bool FPluginReferenceDescriptor::IsEnabledForPlatform(const FString& Platform) const
+bool FPluginReferenceDescriptor::IsEnabledForPlatform( const FString& Platform ) const
 {
 	// If it's not enabled at all, return false
 	if(!bEnabled)
@@ -129,7 +130,8 @@ bool FPluginReferenceDescriptor::IsEnabledForPlatform(const FString& Platform) c
 	return true;
 }
 
-bool FPluginReferenceDescriptor::Read(const FJsonObject& Object, FText& OutFailReason)
+
+bool FPluginReferenceDescriptor::Read( const FJsonObject& Object, FText& OutFailReason )
 {
 	// Get the name
 	if(!Object.TryGetStringField(TEXT("Name"), Name))
@@ -151,72 +153,92 @@ bool FPluginReferenceDescriptor::Read(const FJsonObject& Object, FText& OutFailR
 	// Get the platform lists
 	Object.TryGetStringArrayField(TEXT("WhitelistPlatforms"), WhitelistPlatforms);
 	Object.TryGetStringArrayField(TEXT("BlacklistPlatforms"), BlacklistPlatforms);
+
 	return true;
 }
 
-bool FPluginReferenceDescriptor::ReadArray(const FJsonObject& Object, const TCHAR* Name, TArray<FPluginReferenceDescriptor>& OutPlugins, FText& OutFailReason)
+
+bool FPluginReferenceDescriptor::ReadArray( const FJsonObject& Object, const TCHAR* Name, TArray<FPluginReferenceDescriptor>& OutPlugins, FText& OutFailReason )
 {
 	const TArray< TSharedPtr<FJsonValue> > *Array;
-	if(Object.TryGetArrayField(Name, Array))
+
+	if (Object.TryGetArrayField(Name, Array))
 	{
-		for(const TSharedPtr<FJsonValue> &Item : *Array)
+		for (const TSharedPtr<FJsonValue> &Item : *Array)
 		{
 			const TSharedPtr<FJsonObject> *Object;
-			if(Item.IsValid() && Item->TryGetObject(Object))
+
+			if (Item.IsValid() && Item->TryGetObject(Object))
 			{
 				FPluginReferenceDescriptor Plugin;
-				if(!Plugin.Read(*Object->Get(), OutFailReason))
+
+				if (!Plugin.Read(*Object->Get(), OutFailReason))
 				{
 					return false;
 				}
+
 				OutPlugins.Add(Plugin);
 			}
 		}
 	}
+
 	return true;
 }
 
-void FPluginReferenceDescriptor::Write(TJsonWriter<>& Writer) const
+
+void FPluginReferenceDescriptor::Write( TJsonWriter<>& Writer ) const
 {
 	Writer.WriteObjectStart();
 	Writer.WriteValue(TEXT("Name"), Name);
 	Writer.WriteValue(TEXT("Enabled"), bEnabled);
+
 	if (Description.Len() > 0)
 	{
 		Writer.WriteValue(TEXT("Description"), Description);
 	}
+
 	if (WhitelistPlatforms.Num() > 0)
 	{
 		Writer.WriteArrayStart(TEXT("WhitelistPlatforms"));
-		for(int Idx = 0; Idx < WhitelistPlatforms.Num(); Idx++)
+
+		for (int Idx = 0; Idx < WhitelistPlatforms.Num(); Idx++)
 		{
 			Writer.WriteValue(WhitelistPlatforms[Idx]);
 		}
+
 		Writer.WriteArrayEnd();
 	}
+
 	if (BlacklistPlatforms.Num() > 0)
 	{
 		Writer.WriteArrayStart(TEXT("BlacklistPlatforms"));
-		for(int Idx = 0; Idx < BlacklistPlatforms.Num(); Idx++)
+
+		for (int Idx = 0; Idx < BlacklistPlatforms.Num(); Idx++)
 		{
 			Writer.WriteValue(BlacklistPlatforms[Idx]);
 		}
+
 		Writer.WriteArrayEnd();
 	}
+
 	Writer.WriteObjectEnd();
 }
 
-void FPluginReferenceDescriptor::WriteArray(TJsonWriter<>& Writer, const TCHAR* Name, const TArray<FPluginReferenceDescriptor>& Plugins)
+
+void FPluginReferenceDescriptor::WriteArray( TJsonWriter<>& Writer, const TCHAR* Name, const TArray<FPluginReferenceDescriptor>& Plugins )
 {
-	if(Plugins.Num() > 0)
+	if( Plugins.Num() > 0)
 	{
 		Writer.WriteArrayStart(Name);
-		for(int Idx = 0; Idx < Plugins.Num(); Idx++)
+
+		for (int Idx = 0; Idx < Plugins.Num(); Idx++)
 		{
 			Plugins[Idx].Write(Writer);
 		}
+
 		Writer.WriteArrayEnd();
 	}
 }
 
+
 #undef LOCTEXT_NAMESPACE
diff --git a/Engine/Source/Runtime/Projects/Private/PluginManager.cpp b/Engine/Source/Runtime/Projects/Private/PluginManager.cpp
--- a/Engine/Source/Runtime/Projects/Private/PluginManager.cpp
+++ b/Engine/Source/Runtime/Projects/Private/PluginManager.cpp
@@ -468,6 +468,7 @@ TArray< FPluginStatus > FPluginManager::QueryStatusForAllPlugins() const
 		PluginStatus.CreatedBy = PluginInfo.CreatedBy;
 		PluginStatus.CreatedByURL = PluginInfo.CreatedByURL;
 		PluginStatus.CategoryPath = PluginInfo.Category;
+		PluginStatus.DocsURL = PluginInfo.DocsURL;
 		PluginStatus.PluginDirectory = FPaths::GetPath(Plugin->FileName);
 		PluginStatus.bIsEnabled = Plugin->bEnabled;
 		PluginStatus.bIsBuiltIn = ( Plugin->LoadedFrom == EPluginLoadedFrom::Engine );
diff --git a/Engine/Source/Runtime/Projects/Public/Interfaces/IPluginManager.h b/Engine/Source/Runtime/Projects/Public/Interfaces/IPluginManager.h
--- a/Engine/Source/Runtime/Projects/Public/Interfaces/IPluginManager.h
+++ b/Engine/Source/Runtime/Projects/Public/Interfaces/IPluginManager.h
@@ -2,71 +2,72 @@
 
 #pragma once
 
-
 #include "ModuleDescriptor.h"
 
 
 /**
- * Simple data structure that is filled when querying information about plugins
+ * Simple data structure that is filled when querying information about plug-ins.
  */
 class FPluginStatus
 {
-
 public:
 
-	/** The name of this plugin */
+	/** The name of this plug-in. */
 	FString Name;
 
-	/** Friendly name for the plugin */
+	/** Friendly name for the plug-in. */
 	FString FriendlyName;
 
-	/** Internal version number (not user displayed, but valid for comparisons.) */
+	/** Internal version number (not user displayed, but valid for comparisons). */
 	int32 Version;
 
-	/** Friendly version name */
+	/** Friendly version name. */
 	FString VersionName;
 
-	/** Description of the plugin */
+	/** Description of the plug-in. */
 	FString Description;
 
-	/** Created by name */
+	/** Created by name. */
 	FString CreatedBy;
 
-	/** Created by URL string */
+	/** Created by URL string. */
 	FString CreatedByURL;
 
-	/** Category path (dot-separated list of categories) */
+	/** Category path (dot-separated list of categories). */
 	FString CategoryPath;
 
-	/** Path to plugin directory on disk */
+	/** Documentation URL string. */
+	FString DocsURL;
+
+	/** Path to plug-in directory on disk. */
 	FString PluginDirectory;
 
-	/** True if plugin is currently enabled */
+	/** True if plug-in is currently enabled. */
 	bool bIsEnabled;
 
-	/** True if plugin is enabled by default in all projects */
+	/** True if plug-in is enabled by default in all projects. */
 	bool bIsEnabledByDefault;
 
-	/** True if the plugin is a 'built-in' engine plugin */
+	/** True if the plug-in is a 'built-in' engine plug-in. */
 	bool bIsBuiltIn;
 
-	/** Full path to the 128x128 thumbnail icon file name (or an empty string if no icon is available) */
+	/** Full path to the 128x128 thumbnail icon file name (or an empty string if no icon is available). */
 	FString Icon128FilePath;
 
-	/** Marks the plugin as beta in the UI */
+	/** Marks the plug-in as beta in the UI. */
 	bool bIsBetaVersion;
 
-	/** Whether the plugin has a content folder */
+	/** Whether the plug-in has a content folder. */
 	bool bHasContentFolder;
-
 };
 
+
 /**
- * Structure holding information about a plugin content folder
+ * Structure holding information about a plug-in content folder.
  */
 struct FPluginContentFolder
 {
-	/** Name of the plugin */
+	/** Name of the plug-in */
 	FString Name;
 
 	/** Virtual root path for asset paths */
@@ -76,26 +77,18 @@ struct FPluginContentFolder
 	FString ContentPath;
 };
 
+
 /**
- * PluginManager manages available code and content extensions (both loaded and not loaded.)
+ * PluginManager manages available code and content extensions (both loaded and not loaded).
  */
 class IPluginManager
 {
-
 public:
 
 	/**
-	 * Static: Access singleton instance
-	 *
-	 * @return	Reference to the singleton object
-	 */
-	static PROJECTS_API IPluginManager& Get();
-
-
-	/**
-	 * Loads all plugins
+	 * Loads all plug-ins
 	 *
-	 * @param	LoadingPhase	Which loading phase we're loading plugin modules from.  Only modules that are configured to be
+	 * @param	LoadingPhase	Which loading phase we're loading plug-in modules from.  Only modules that are configured to be
 	 *							loaded at the specified loading phase will be loaded during this call.
 	 */
 	virtual bool LoadModulesForEnabledPlugins( const ELoadingPhase::Type LoadingPhase ) = 0;
@@ -104,39 +97,48 @@ public:
 	DECLARE_DELEGATE_TwoParams( FRegisterMountPointDelegate, const FString& /* Root content path */, const FString& /* Directory name */ );
 
 	/**
-	 * Sets the delegate to call to register a new content mount point.  This is used internally by the plugin manager system
+	 * Sets the delegate to call to register a new content mount point.  This is used internally by the plug-in manager system
 	 * and should not be called by you.  This is registered at application startup by FPackageName code in CoreUObject.
 	 *
-	 * @param	Delegate	The delegate to that will be called when plugin manager needs to register a mount point
+	 * @param	Delegate	The delegate to that will be called when plug-in manager needs to register a mount point
 	 */
 	virtual void SetRegisterMountPointDelegate( const FRegisterMountPointDelegate& Delegate ) = 0;
 
 	/**
-	 * Checks if all the required plugins are available. If not, will present an error dialog the first time a plugin is loaded or this function is called.
+	 * Checks if all the required plug-ins are available. If not, will present an error dialog the first time a plug-in is loaded or this function is called.
 	 *
-	 * @returns true if all the required plugins are available.
+	 * @returns true if all the required plug-ins are available.
 	 */
-	virtual bool AreRequiredPluginsAvailable( ) = 0;
+	virtual bool AreRequiredPluginsAvailable() = 0;
 
 	/** 
-	 * Checks whether modules for the enabled plugins are up to date
+	 * Checks whether modules for the enabled plug-ins are up to date.
 	 *
-	 * @param OutIncompatibleNames	Array to receive a list of incompatible module names
-	 * @returns true if the enabled plugin modules are up to date
+	 * @param OutIncompatibleNames	Array to receive a list of incompatible module names.
+	 * @returns true if the enabled plug-in modules are up to date.
 	 */
 	virtual bool CheckModuleCompatibility( TArray<FString>& OutIncompatibleModules ) = 0;
 
 	/**
-	 * Gets status about all currently known plugins
+	 * Gets status about all currently known plug-ins.
+	 *
+	 * @return	 Array of plug-in status objects.
+	 */
+	virtual TArray<FPluginStatus> QueryStatusForAllPlugins() const = 0;
+
+	/**
+	 * Gets a list of plug-in content folders.
 	 *
-	 * @return	 Array of plugin status objects
+	 * @return	 Array of plug-in content folders.
 	 */
-	virtual TArray< FPluginStatus > QueryStatusForAllPlugins() const = 0;
+	virtual const TArray<FPluginContentFolder>& GetPluginContentFolders() const = 0;
+
+public:
 
 	/**
-	 * Gets a list of plugin content folders
+	 * Static: Access singleton instance.
 	 *
-	 * @return	 Array of plugin content folders
+	 * @return	Reference to the singleton object.
 	 */
-	virtual const TArray< FPluginContentFolder >& GetPluginContentFolders() const = 0;
+	static PROJECTS_API IPluginManager& Get();
 };
diff --git a/Engine/Source/Runtime/Projects/Public/PluginDescriptor.h b/Engine/Source/Runtime/Projects/Public/PluginDescriptor.h
--- a/Engine/Source/Runtime/Projects/Public/PluginDescriptor.h
+++ b/Engine/Source/Runtime/Projects/Public/PluginDescriptor.h
@@ -56,6 +56,9 @@ struct PROJECTS_API FPluginDescriptor
 	/** Hyperlink URL string for the company or individual who created this plugin.  This is optional. */
 	FString CreatedByURL;
 
+	/** Documentation URL string. */
+	FString DocsURL;
+
 	/** List of all modules associated with this plugin */
 	TArray<FModuleDescriptor> Modules;
 
