Remove the recursive display from the plugin browser, and rename the top-level categories. Plugins in the engine folder are now under "Installed", plugins in the project folder are now under "Project".

diff --git a/Engine/Plugins/Editor/PluginsEditor/Source/PluginsEditor/Private/SPluginCategories.cpp b/Engine/Plugins/Editor/PluginsEditor/Source/PluginsEditor/Private/SPluginCategories.cpp
--- a/Engine/Plugins/Editor/PluginsEditor/Source/PluginsEditor/Private/SPluginCategories.cpp
+++ b/Engine/Plugins/Editor/PluginsEditor/Source/PluginsEditor/Private/SPluginCategories.cpp
@@ -66,10 +66,10 @@ void SPluginCategories::RebuildAndFilterCategoryTree()
 		RootPluginCategories.Reset();
 
 		// Add the root level categories
-		TSharedRef<FPluginCategoryTreeItem> BuiltInCategory = MakeShareable(new FPluginCategoryTreeItem(NULL, TEXT("BuiltIn"), TEXT("Built-in"), LOCTEXT("BuiltInCategoryName", "Built-in")));
-		RootPluginCategories.Add( BuiltInCategory );
 		TSharedRef<FPluginCategoryTreeItem> InstalledCategory = MakeShareable(new FPluginCategoryTreeItem(NULL, TEXT("Installed"), TEXT("Installed"), LOCTEXT("InstalledCategoryName", "Installed")));
 		RootPluginCategories.Add( InstalledCategory );
+		TSharedRef<FPluginCategoryTreeItem> ProjectCategory = MakeShareable(new FPluginCategoryTreeItem(NULL, TEXT("Project"), TEXT("Project"), LOCTEXT("ProjectCategoryName", "Project")));
+		RootPluginCategories.Add( ProjectCategory );
 
 
 		const TArray< FPluginStatus > Plugins = IPluginManager::Get().QueryStatusForAllPlugins();
@@ -82,11 +82,11 @@ void SPluginCategories::RebuildAndFilterCategoryTree()
 			FPluginCategoryTreeItemPtr CategoryForPlugin;
 			if( PluginStatus.bIsBuiltIn )
 			{
-				CategoryForPlugin = BuiltInCategory;
+				CategoryForPlugin = InstalledCategory;
 			}
 			else
 			{
-				CategoryForPlugin = InstalledCategory;
+				CategoryForPlugin = ProjectCategory;
 			}
 			check( CategoryForPlugin.IsValid() );
 			FString ItemCategoryPath = CategoryForPlugin->GetCategoryPath();
diff --git a/Engine/Plugins/Editor/PluginsEditor/Source/PluginsEditor/Private/SPluginList.cpp b/Engine/Plugins/Editor/PluginsEditor/Source/PluginsEditor/Private/SPluginList.cpp
--- a/Engine/Plugins/Editor/PluginsEditor/Source/PluginsEditor/Private/SPluginList.cpp
+++ b/Engine/Plugins/Editor/PluginsEditor/Source/PluginsEditor/Private/SPluginList.cpp
@@ -65,15 +65,8 @@ TSharedRef<ITableRow> SPluginList::PluginListView_OnGenerateRow( FPluginListItem
 
 
 
-void SPluginList::GetPluginsRecursively( const TSharedPtr< FPluginCategoryTreeItem >& Category )
+void SPluginList::GetPlugins( const TSharedPtr< FPluginCategoryTreeItem >& Category )
 {
-	// Process sub-categories
-	for( auto SubCategoryIt( Category->GetSubCategories().CreateConstIterator() ); SubCategoryIt; ++SubCategoryIt )
-	{
-		const auto& SubCategory = *SubCategoryIt;
-		GetPluginsRecursively( SubCategory );
-	}
-
 	// Add plugins from this category
 	const TArray< TSharedPtr< FPluginStatus > >& Plugins = Category->GetPlugins();
 	for( auto PluginIt( Plugins.CreateConstIterator() ); PluginIt; ++PluginIt )
@@ -107,7 +100,7 @@ void SPluginList::RebuildAndFilterPluginList()
 		const auto& SelectedCategory = OwnerWeak.Pin()->GetSelectedCategory();
 		if( SelectedCategory.IsValid() )
 		{
-			GetPluginsRecursively( SelectedCategory );
+			GetPlugins( SelectedCategory );
 		}
 
 		// Sort the plugins alphabetically
diff --git a/Engine/Plugins/Editor/PluginsEditor/Source/PluginsEditor/Private/SPluginList.h b/Engine/Plugins/Editor/PluginsEditor/Source/PluginsEditor/Private/SPluginList.h
--- a/Engine/Plugins/Editor/PluginsEditor/Source/PluginsEditor/Private/SPluginList.h
+++ b/Engine/Plugins/Editor/PluginsEditor/Source/PluginsEditor/Private/SPluginList.h
@@ -37,7 +37,7 @@ private:
 	void OnPluginTextFilterChanged();
 
 	/** Recursively gathers plugins from a category and all sub-categories.  Updates the PluginListItems structure. */
-	void GetPluginsRecursively( const TSharedPtr< class FPluginCategoryTreeItem >& Category );
+	void GetPlugins( const TSharedPtr< class FPluginCategoryTreeItem >& Category );
 
 	/** Called to generate a widget for the specified list item */
 	TSharedRef<ITableRow> PluginListView_OnGenerateRow( FPluginListItemPtr Item, const TSharedRef<STableViewBase>& OwnerTable );
