Remove some of the tedious layers of indirection from the plugin manager.

diff --git a/Engine/Plugins/Editor/PluginsEditor/Source/PluginsEditor/Private/PluginsEditorPrivatePCH.h b/Engine/Plugins/Editor/PluginsEditor/Source/PluginsEditor/Private/PluginsEditorPrivatePCH.h
--- a/Engine/Plugins/Editor/PluginsEditor/Source/PluginsEditor/Private/PluginsEditorPrivatePCH.h
+++ b/Engine/Plugins/Editor/PluginsEditor/Source/PluginsEditor/Private/PluginsEditorPrivatePCH.h
@@ -7,4 +7,4 @@
 
 #include "TextFilter.h"
 
-typedef TTextFilter< const class FPluginStatus& > FPluginTextFilter;
+typedef TTextFilter< const struct FPluginStatus& > FPluginTextFilter;
diff --git a/Engine/Plugins/Editor/PluginsEditor/Source/PluginsEditor/Private/SPluginCategories.cpp b/Engine/Plugins/Editor/PluginsEditor/Source/PluginsEditor/Private/SPluginCategories.cpp
--- a/Engine/Plugins/Editor/PluginsEditor/Source/PluginsEditor/Private/SPluginCategories.cpp
+++ b/Engine/Plugins/Editor/PluginsEditor/Source/PluginsEditor/Private/SPluginCategories.cpp
@@ -73,14 +73,11 @@ void SPluginCategories::RebuildAndFilterCategoryTree()
 
 
 		const TArray< FPluginStatus > Plugins = IPluginManager::Get().QueryStatusForAllPlugins();
-		for( auto PluginIt( Plugins.CreateConstIterator() ); PluginIt; ++PluginIt )
+		for(const FPluginStatus& Plugin: Plugins)
 		{
-			const auto& PluginStatus = *PluginIt;
-
-
 			// Figure out which base category this plugin belongs in
 			FPluginCategoryTreeItemPtr CategoryForPlugin;
-			if( PluginStatus.bIsBuiltIn )
+			if( Plugin.LoadedFrom == EPluginLoadedFrom::Engine )
 			{
 				CategoryForPlugin = InstalledCategory;
 			}
@@ -92,7 +89,7 @@ void SPluginCategories::RebuildAndFilterCategoryTree()
 			FString ItemCategoryPath = CategoryForPlugin->GetCategoryPath();
 
 
-			const FString& CategoryPath = PluginStatus.CategoryPath;
+			const FString& CategoryPath = Plugin.Descriptor.Category;
 			{
 				// We're expecting the category string to be in the "A.B.C" format.  We'll split up the string here and form
 				// a proper hierarchy in the UI
@@ -143,7 +140,7 @@ void SPluginCategories::RebuildAndFilterCategoryTree()
 			
 			// Associate the plugin with the category
 			// PERFORMANCE NOTE: Copying FPluginStats by value here
-			CategoryForPlugin->AddPlugin( MakeShareable( new FPluginStatus( PluginStatus ) ) );
+			CategoryForPlugin->AddPlugin( MakeShareable( new FPluginStatus( Plugin ) ) );
 		}
 
 
diff --git a/Engine/Plugins/Editor/PluginsEditor/Source/PluginsEditor/Private/SPluginList.cpp b/Engine/Plugins/Editor/PluginsEditor/Source/PluginsEditor/Private/SPluginList.cpp
--- a/Engine/Plugins/Editor/PluginsEditor/Source/PluginsEditor/Private/SPluginList.cpp
+++ b/Engine/Plugins/Editor/PluginsEditor/Source/PluginsEditor/Private/SPluginList.cpp
@@ -109,7 +109,7 @@ void SPluginList::RebuildAndFilterPluginList()
 			{
 				bool operator()( const FPluginListItemPtr& A, const FPluginListItemPtr& B ) const
 				{
-					return A->PluginStatus.FriendlyName < B->PluginStatus.FriendlyName;
+					return A->PluginStatus.Descriptor.FriendlyName < B->PluginStatus.Descriptor.FriendlyName;
 				}
 			};
 			PluginListItems.Sort( FPluginListItemSorter() );
diff --git a/Engine/Plugins/Editor/PluginsEditor/Source/PluginsEditor/Private/SPluginListTile.cpp b/Engine/Plugins/Editor/PluginsEditor/Source/PluginsEditor/Private/SPluginListTile.cpp
--- a/Engine/Plugins/Editor/PluginsEditor/Source/PluginsEditor/Private/SPluginListTile.cpp
+++ b/Engine/Plugins/Editor/PluginsEditor/Source/PluginsEditor/Private/SPluginListTile.cpp
@@ -24,10 +24,13 @@ void SPluginListTile::Construct( const FArguments& Args, const TSharedRef<SPlugi
 	//		-> Maybe a filter for this too?  (show only editor plugins, etc.)
 	// @todo plugedit: Indicate whether plugin has content?  Filter to show only content plugins, and vice-versa?
 
+	// @todo plugedit: Maybe we should do the FileExists check ONCE at plugin load time and not at query time
+
 	// Plugin thumbnail image
-	if( !Item->PluginStatus.Icon128FilePath.IsEmpty() )
+	const FString Icon128FilePath = Item->PluginStatus.PluginDirectory / TEXT("Resources/Icon128.png");
+	if(FPlatformFileManager::Get().GetPlatformFile().FileExists(*Icon128FilePath))
 	{
-		const FName BrushName( *Item->PluginStatus.Icon128FilePath );
+		const FName BrushName( *Icon128FilePath );
 		const FIntPoint Size = FSlateApplication::Get().GetRenderer()->GenerateDynamicImageResource(BrushName);
 
 		if ((Size.X > 0) && (Size.Y > 0))
@@ -44,13 +47,13 @@ void SPluginListTile::Construct( const FArguments& Args, const TSharedRef<SPlugi
 	// create documentation link
 	TSharedPtr<SWidget> DocumentationWidget;
 	{
-		if (Item->PluginStatus.DocsURL.IsEmpty())
+		if (Item->PluginStatus.Descriptor.DocsURL.IsEmpty())
 		{
 			DocumentationWidget = SNullWidget::NullWidget;
 		}
 		else
 		{
-			FString DocsURL = Item->PluginStatus.DocsURL;
+			FString DocsURL = Item->PluginStatus.Descriptor.DocsURL;
 			DocumentationWidget = SNew(SHorizontalBox)
 
 			+ SHorizontalBox::Slot()
@@ -78,11 +81,11 @@ void SPluginListTile::Construct( const FArguments& Args, const TSharedRef<SPlugi
 	// create vendor link
 	TSharedPtr<SWidget> CreatedByWidget;
 	{
-		if (Item->PluginStatus.CreatedBy.IsEmpty())
+		if (Item->PluginStatus.Descriptor.CreatedBy.IsEmpty())
 		{
 			CreatedByWidget = SNullWidget::NullWidget;
 		}
-		else if (Item->PluginStatus.CreatedByURL.IsEmpty())
+		else if (Item->PluginStatus.Descriptor.CreatedByURL.IsEmpty())
 		{
 			CreatedByWidget = SNew(SHorizontalBox)
 
@@ -101,12 +104,12 @@ void SPluginListTile::Construct( const FArguments& Args, const TSharedRef<SPlugi
 				.Padding(2.0f, 0.0f, 0.0f, 0.0f)
 				[
 					SNew(STextBlock)
-						.Text(FText::FromString(Item->PluginStatus.CreatedBy))
+						.Text(FText::FromString(Item->PluginStatus.Descriptor.CreatedBy))
 				];
 		}
 		else
 		{
-			FString CreatedByURL = Item->PluginStatus.CreatedByURL;
+			FString CreatedByURL = Item->PluginStatus.Descriptor.CreatedByURL;
 			CreatedByWidget = SNew(SHorizontalBox)
 
 			+ SHorizontalBox::Slot()
@@ -124,7 +127,7 @@ void SPluginListTile::Construct( const FArguments& Args, const TSharedRef<SPlugi
 				.Padding(2.0f, 0.0f, 0.0f, 0.0f)
 				[				
 					SNew(SHyperlink)
-						.Text(FText::FromString(Item->PluginStatus.CreatedBy))
+						.Text(FText::FromString(Item->PluginStatus.Descriptor.CreatedBy))
 						.ToolTipText(FText::Format(LOCTEXT("NavigateToCreatedByURL", "Visit the vendor's web site ({0})"), FText::FromString(CreatedByURL)))
 						.OnNavigate_Lambda([=]() { FPlatformProcess::LaunchURL(*CreatedByURL, nullptr, nullptr); })
 				];
@@ -171,7 +174,7 @@ void SPluginListTile::Construct( const FArguments& Args, const TSharedRef<SPlugi
 										.Padding(PaddingAmount)
 										[
 											SNew(STextBlock)
-												.Text(FText::FromString(Item->PluginStatus.FriendlyName))
+												.Text(FText::FromString(Item->PluginStatus.Descriptor.FriendlyName))
 												.HighlightText_Raw(&Owner->GetOwner().GetPluginTextFilter(), &FPluginTextFilter::GetRawFilterText)
 												.TextStyle(FPluginStyle::Get(), "PluginTile.NameText")
 										]
@@ -197,7 +200,7 @@ void SPluginListTile::Construct( const FArguments& Args, const TSharedRef<SPlugi
 															.Padding(0.0f, 0.0f, 3.0f, 2.0f)
 															[
 																SNew(SImage)
-																	.Visibility(Item->PluginStatus.bIsBetaVersion ? EVisibility::Visible : EVisibility::Collapsed)
+																	.Visibility(Item->PluginStatus.Descriptor.bIsBetaVersion ? EVisibility::Visible : EVisibility::Collapsed)
 																	.Image(FPluginStyle::Get()->GetBrush("PluginTile.BetaWarning"))
 															]
 
@@ -207,7 +210,7 @@ void SPluginListTile::Construct( const FArguments& Args, const TSharedRef<SPlugi
 															.VAlign(VAlign_Bottom)
 															[
 																SNew(STextBlock)
-																	.Text(Item->PluginStatus.bIsBetaVersion ? LOCTEXT("PluginBetaVersionLabel", "BETA Version ") : LOCTEXT("PluginVersionLabel", "Version "))
+																	.Text(Item->PluginStatus.Descriptor.bIsBetaVersion ? LOCTEXT("PluginBetaVersionLabel", "BETA Version ") : LOCTEXT("PluginVersionLabel", "Version "))
 															]
 													]
 
@@ -217,7 +220,7 @@ void SPluginListTile::Construct( const FArguments& Args, const TSharedRef<SPlugi
 													.Padding( 0.0f, 0.0f, 2.0f, 0.0f )	// Extra padding from the right edge
 													[
 														SNew(STextBlock)
-															.Text(FText::FromString(Item->PluginStatus.VersionName))
+															.Text(FText::FromString(Item->PluginStatus.Descriptor.VersionName))
 															.TextStyle(FPluginStyle::Get(), "PluginTile.VersionNumberText")
 													]
 											]
@@ -232,7 +235,7 @@ void SPluginListTile::Construct( const FArguments& Args, const TSharedRef<SPlugi
 											.Padding( PaddingAmount )
 											[
 												SNew(STextBlock)
-													.Text(FText::FromString(Item->PluginStatus.Description))
+													.Text(FText::FromString(Item->PluginStatus.Descriptor.Description))
 													.AutoWrapText(true)
 											]
 
@@ -258,24 +261,6 @@ void SPluginListTile::Construct( const FArguments& Args, const TSharedRef<SPlugi
 															]
 													]
 
-												// Uninstall button
-												+ SHorizontalBox::Slot()
-													.Padding(PaddingAmount)
-													.AutoWidth()
-													[
-														SNew(SButton)
-															.Text(LOCTEXT("PluginUninstallButtonLabel", "Uninstall..."))
-
-															// Don't show 'Uninstall' for built-in plugins
-															.Visibility(EVisibility::Collapsed)
-															// @todo plugins: Not supported yet
-																// .Visibility( Item->PluginStatus.bIsBuiltIn ? EVisibility::Collapsed : EVisibility::Visible )
-
-															.OnClicked(this, &SPluginListTile::OnUninstallButtonClicked)
-															.IsEnabled(false)		// @todo plugedit: Not supported yet
-															.ToolTipText(LOCTEXT("UninstallButtonToolTip", "Allows you to uninstall this plugin.  This will delete the plugin from your project entirely.  You may need to restart the program for this change to take effect."))
-													]
-
 												// docs link
 												+ SHorizontalBox::Slot()
 													.Padding(PaddingAmount)
@@ -313,11 +298,11 @@ void SPluginListTile::OnEnablePluginCheckboxChanged(ECheckBoxState NewCheckedSta
 {
 	const bool bNewEnabledState = (NewCheckedState == ECheckBoxState::Checked);
 
-	if (bNewEnabledState && ItemData->PluginStatus.bIsBetaVersion)
+	if (bNewEnabledState && ItemData->PluginStatus.Descriptor.bIsBetaVersion)
 	{
 		FText WarningMessage = FText::Format(
 			LOCTEXT("Warning_EnablingBetaPlugin", "Plugin '{0}' is a beta version and might be unstable or removed without notice. Please use with caution. Are you sure you want to enable the plugin?"),
-			FText::FromString(ItemData->PluginStatus.FriendlyName));
+			FText::FromString(ItemData->PluginStatus.Descriptor.FriendlyName));
 
 		if (EAppReturnType::No == FMessageDialog::Open(EAppMsgType::YesNo, WarningMessage))
 		{
@@ -336,17 +321,4 @@ void SPluginListTile::OnEnablePluginCheckboxChanged(ECheckBoxState NewCheckedSta
 	}
 }
 
-
-FReply SPluginListTile::OnUninstallButtonClicked()
-{
-	// not expecting to deal with engine plugins here (they can't be uninstalled)
-	if (ensure(!ItemData->PluginStatus.bIsBuiltIn))
-	{
-		// @todo plugedit: Make this work!
-	}
-
-	return FReply::Handled();
-}
-
-
 #undef LOCTEXT_NAMESPACE
diff --git a/Engine/Plugins/Editor/PluginsEditor/Source/PluginsEditor/Private/SPluginListTile.h b/Engine/Plugins/Editor/PluginsEditor/Source/PluginsEditor/Private/SPluginListTile.h
--- a/Engine/Plugins/Editor/PluginsEditor/Source/PluginsEditor/Private/SPluginListTile.h
+++ b/Engine/Plugins/Editor/PluginsEditor/Source/PluginsEditor/Private/SPluginListTile.h
@@ -26,9 +26,6 @@ private:
 	/** Called when the enabled checkbox is clicked */
 	void OnEnablePluginCheckboxChanged(ECheckBoxState NewCheckedState);
 
-	/** Called when 'Uninstall' is clicked on a plugin's tile */
-	FReply OnUninstallButtonClicked();
-
 
 private:
 
diff --git a/Engine/Plugins/Editor/PluginsEditor/Source/PluginsEditor/Private/SPluginsEditor.cpp b/Engine/Plugins/Editor/PluginsEditor/Source/PluginsEditor/Private/SPluginsEditor.cpp
--- a/Engine/Plugins/Editor/PluginsEditor/Source/PluginsEditor/Private/SPluginsEditor.cpp
+++ b/Engine/Plugins/Editor/PluginsEditor/Source/PluginsEditor/Private/SPluginsEditor.cpp
@@ -20,7 +20,7 @@ void SPluginsEditor::Construct( const FArguments& Args )
 		static void PluginStatusToStringArray( const FPluginStatus& PluginStatus, OUT TArray< FString >& StringArray )
 		{
 			// NOTE: Only the friendly name is searchable for now.  We don't display the actual plugin name in the UI.
-			StringArray.Add( PluginStatus.FriendlyName );
+			StringArray.Add( PluginStatus.Descriptor.FriendlyName );
 		}
 	};
 
diff --git a/Engine/Source/Editor/ContentBrowser/Private/NativeClassHierarchy.cpp b/Engine/Source/Editor/ContentBrowser/Private/NativeClassHierarchy.cpp
--- a/Engine/Source/Editor/ContentBrowser/Private/NativeClassHierarchy.cpp
+++ b/Engine/Source/Editor/ContentBrowser/Private/NativeClassHierarchy.cpp
@@ -488,9 +488,9 @@ TMap<FName, FName> FNativeClassHierarchy::GetPluginModules()
 	TMap<FName, FName> PluginModules;
 	{
 		TArray<FPluginStatus> Plugins = PluginManager.QueryStatusForAllPlugins();
-		for(const auto& Plugin : Plugins)
+		for(const FPluginStatus& Plugin: Plugins)
 		{
-			for(const auto& PluginModule : Plugin.Modules)
+			for(const FModuleDescriptor& PluginModule: Plugin.Descriptor.Modules)
 			{
 				PluginModules.Add(PluginModule.Name, FName(*Plugin.Name));
 			}
diff --git a/Engine/Source/Programs/UnrealHeaderTool/Private/CodeGenerator.cpp b/Engine/Source/Programs/UnrealHeaderTool/Private/CodeGenerator.cpp
--- a/Engine/Source/Programs/UnrealHeaderTool/Private/CodeGenerator.cpp
+++ b/Engine/Source/Programs/UnrealHeaderTool/Private/CodeGenerator.cpp
@@ -5095,7 +5095,7 @@ void GetScriptPlugins(TArray<IScriptGeneratorPluginInterface*>& ScriptPlugins)
 	for (auto PluginIt(Plugins.CreateConstIterator()); PluginIt; ++PluginIt)
 	{
 		const auto& PluginStatus = *PluginIt;
-		if (PluginStatus.bIsEnabled && PluginStatus.CategoryPath.StartsWith(CodeGeneratorPluginCategory))
+		if (PluginStatus.bIsEnabled && PluginStatus.Descriptor.Category.StartsWith(CodeGeneratorPluginCategory))
 		{
 			auto GeneratorInterface = FModuleManager::LoadModulePtr<IScriptGeneratorPluginInterface>(*PluginStatus.Name);
 			if (GeneratorInterface)
diff --git a/Engine/Source/Runtime/Projects/Private/PluginManager.cpp b/Engine/Source/Runtime/Projects/Private/PluginManager.cpp
--- a/Engine/Source/Runtime/Projects/Private/PluginManager.cpp
+++ b/Engine/Source/Runtime/Projects/Private/PluginManager.cpp
@@ -13,12 +13,9 @@ namespace PluginSystemDefs
 	    NOTE: This constant exists in UnrealBuildTool code as well. */
 	static const TCHAR PluginDescriptorFileExtension[] = TEXT( ".uplugin" );
 
-	/** Relative path to the plugin's 128x128 icon resource file */
-	static const FString RelativeIcon128FilePath( TEXT( "Resources/Icon128.png" ) );
-
 }
 
-FPluginInstance::FPluginInstance(const FString& InFileName, const FPluginDescriptor& InDescriptor, EPluginLoadedFrom::Type InLoadedFrom)
+FPluginInstance::FPluginInstance(const FString& InFileName, const FPluginDescriptor& InDescriptor, EPluginLoadedFrom InLoadedFrom)
 	: Name(FPaths::GetBaseFilename(InFileName))
 	, FileName(InFileName)
 	, Descriptor(InDescriptor)
@@ -66,7 +63,7 @@ void FPluginManager::DiscoverAllPlugins()
 		 * @param	LoadedFrom			Where we're loading these plugins from (game, engine, etc)
 		 * @param	Plugins				The array to be filled in with new plugins including descriptors
 		 */
-		static void FindPluginsRecursively( const FString& PluginsDirectory, const EPluginLoadedFrom::Type LoadedFrom, TArray< TSharedRef<FPluginInstance> >& Plugins )
+		static void FindPluginsRecursively( const FString& PluginsDirectory, const EPluginLoadedFrom LoadedFrom, TArray< TSharedRef<FPluginInstance> >& Plugins )
 		{
 			// NOTE: The logic in this function generally matches that of the C# code for FindPluginsRecursively
 			//       in UnrealBuildTool.  These routines should be kept in sync.
@@ -211,7 +208,7 @@ void FPluginManager::DiscoverAllPlugins()
 		 * @param	LoadedFrom			Where we're loading these plugins from (game, engine, etc)
 		 * @param	Plugins				The array to be filled in with loaded plugin descriptors
 		 */
-		static void FindPluginsIn( const FString& PluginsDirectory, const EPluginLoadedFrom::Type LoadedFrom, TArray< TSharedRef<FPluginInstance> >& Plugins )
+		static void FindPluginsIn( const FString& PluginsDirectory, const EPluginLoadedFrom LoadedFrom, TArray< TSharedRef<FPluginInstance> >& Plugins )
 		{
 			// Make sure the directory even exists
 			if( FPlatformFileManager::Get().GetPlatformFile().DirectoryExists( *PluginsDirectory ) )
@@ -516,32 +513,13 @@ TArray< FPluginStatus > FPluginManager::QueryStatusForAllPlugins() const
 	for( auto PluginIt( AllPlugins.CreateConstIterator() ); PluginIt; ++PluginIt )
 	{
 		const TSharedRef< FPluginInstance >& Plugin = *PluginIt;
-		const FPluginDescriptor& PluginInfo = Plugin->Descriptor;
 		
 		FPluginStatus PluginStatus;
 		PluginStatus.Name = Plugin->Name;
-		PluginStatus.FriendlyName = PluginInfo.FriendlyName;
-		PluginStatus.Version = PluginInfo.Version;
-		PluginStatus.VersionName = PluginInfo.VersionName;
-		PluginStatus.Description = PluginInfo.Description;
-		PluginStatus.CreatedBy = PluginInfo.CreatedBy;
-		PluginStatus.CreatedByURL = PluginInfo.CreatedByURL;
-		PluginStatus.CategoryPath = PluginInfo.Category;
-		PluginStatus.DocsURL = PluginInfo.DocsURL;
 		PluginStatus.PluginDirectory = FPaths::GetPath(Plugin->FileName);
 		PluginStatus.bIsEnabled = Plugin->bEnabled;
-		PluginStatus.bIsBuiltIn = ( Plugin->LoadedFrom == EPluginLoadedFrom::Engine );
-		PluginStatus.bIsEnabledByDefault = PluginInfo.bEnabledByDefault;
-		PluginStatus.bIsBetaVersion = PluginInfo.bIsBetaVersion;
-		PluginStatus.bHasContentFolder = PluginInfo.bCanContainContent;
-		PluginStatus.Modules = PluginInfo.Modules;
-
-		// @todo plugedit: Maybe we should do the FileExists check ONCE at plugin load time and not at query time
-		const FString Icon128FilePath = FPaths::GetPath(Plugin->FileName) / PluginSystemDefs::RelativeIcon128FilePath;
-		if( FPlatformFileManager::Get().GetPlatformFile().FileExists( *Icon128FilePath ) )
-		{
-			PluginStatus.Icon128FilePath = Icon128FilePath;
-		}
+		PluginStatus.Descriptor = Plugin->Descriptor;
+		PluginStatus.LoadedFrom = Plugin->LoadedFrom;
 
 		PluginStatuses.Add( PluginStatus );
 	}
diff --git a/Engine/Source/Runtime/Projects/Private/PluginManager.h b/Engine/Source/Runtime/Projects/Private/PluginManager.h
--- a/Engine/Source/Runtime/Projects/Private/PluginManager.h
+++ b/Engine/Source/Runtime/Projects/Private/PluginManager.h
@@ -4,21 +4,6 @@
 
 #include "PluginDescriptor.h"
 
-/**
- * Enum for where a plugin is loaded from
- */
-struct EPluginLoadedFrom
-{
-	enum Type
-	{
-		/** Plugin is built-in to the engine */
-		Engine,
-
-		/** Project-specific plugin, stored within a game project directory */
-		GameProject
-	};
-};
-
 /**
  * Instance of a plugin in memory
  */
@@ -35,7 +20,7 @@ public:
 	FPluginDescriptor Descriptor;
 
 	/** Where does this plugin live? */
-	EPluginLoadedFrom::Type LoadedFrom;
+	EPluginLoadedFrom LoadedFrom;
 
 	/** True if the plugin is marked as enabled */
 	bool bEnabled;
@@ -43,7 +28,7 @@ public:
 	/**
 	 * FPlugin constructor
 	 */
-	FPluginInstance(const FString &FileName, const FPluginDescriptor& InDescriptor, EPluginLoadedFrom::Type InLoadedFrom);
+	FPluginInstance(const FString &FileName, const FPluginDescriptor& InDescriptor, EPluginLoadedFrom InLoadedFrom);
 };
 
 /**
diff --git a/Engine/Source/Runtime/Projects/Private/ProjectManager.cpp b/Engine/Source/Runtime/Projects/Private/ProjectManager.cpp
--- a/Engine/Source/Runtime/Projects/Private/ProjectManager.cpp
+++ b/Engine/Source/Runtime/Projects/Private/ProjectManager.cpp
@@ -267,7 +267,7 @@ bool FProjectManager::IsNonDefaultPluginEnabled() const
 
 	for(const FPluginStatus& Plugin: IPluginManager::Get().QueryStatusForAllPlugins())
 	{
-		if((!Plugin.bIsBuiltIn || !Plugin.bIsEnabledByDefault) && EnabledPlugins.Contains(Plugin.Name))
+		if((Plugin.LoadedFrom == EPluginLoadedFrom::GameProject || !Plugin.Descriptor.bEnabledByDefault) && EnabledPlugins.Contains(Plugin.Name))
 		{
 			return true;
 		}
@@ -339,7 +339,7 @@ void FProjectManager::GetDefaultEnabledPlugins(TArray<FString>& OutPluginNames)
 	TArray<FPluginStatus> PluginStatuses = IPluginManager::Get().QueryStatusForAllPlugins();
 	for(const FPluginStatus& PluginStatus: PluginStatuses)
 	{
-		if(PluginStatus.bIsEnabledByDefault || !PluginStatus.bIsBuiltIn)
+		if(PluginStatus.Descriptor.bEnabledByDefault || PluginStatus.LoadedFrom == EPluginLoadedFrom::GameProject)
 		{
 			OutPluginNames.AddUnique(PluginStatus.Name);
 		}
diff --git a/Engine/Source/Runtime/Projects/Public/Interfaces/IPluginManager.h b/Engine/Source/Runtime/Projects/Public/Interfaces/IPluginManager.h
--- a/Engine/Source/Runtime/Projects/Public/Interfaces/IPluginManager.h
+++ b/Engine/Source/Runtime/Projects/Public/Interfaces/IPluginManager.h
@@ -3,65 +3,40 @@
 #pragma once
 
 #include "ModuleDescriptor.h"
+#include "PluginDescriptor.h"
+
+/**
+ * Enum for where a plugin is loaded from
+ */
+enum class EPluginLoadedFrom
+{
+	/** Plugin is built-in to the engine */
+	Engine,
+
+	/** Project-specific plugin, stored within a game project directory */
+	GameProject
+};
 
 
 /**
  * Simple data structure that is filled when querying information about plug-ins.
  */
-class FPluginStatus
+struct FPluginStatus
 {
-public:
-
 	/** The name of this plug-in. */
 	FString Name;
 
-	/** Friendly name for the plug-in. */
-	FString FriendlyName;
-
-	/** Internal version number (not user displayed, but valid for comparisons). */
-	int32 Version;
-
-	/** Friendly version name. */
-	FString VersionName;
-
-	/** Description of the plug-in. */
-	FString Description;
-
-	/** Created by name. */
-	FString CreatedBy;
-
-	/** Created by URL string. */
-	FString CreatedByURL;
-
-	/** Category path (dot-separated list of categories). */
-	FString CategoryPath;
-
-	/** Documentation URL string. */
-	FString DocsURL;
-
 	/** Path to plug-in directory on disk. */
 	FString PluginDirectory;
 
 	/** True if plug-in is currently enabled. */
 	bool bIsEnabled;
 
-	/** True if plug-in is enabled by default in all projects. */
-	bool bIsEnabledByDefault;
-
-	/** True if the plug-in is a 'built-in' engine plug-in. */
-	bool bIsBuiltIn;
-
-	/** Full path to the 128x128 thumbnail icon file name (or an empty string if no icon is available). */
-	FString Icon128FilePath;
-
-	/** Marks the plug-in as beta in the UI. */
-	bool bIsBetaVersion;
-
-	/** Whether the plug-in has a content folder. */
-	bool bHasContentFolder;
+	/** Where the plugin was loaded from */
+	EPluginLoadedFrom LoadedFrom;
 
-	/** List of all modules associated with this plugin */
-	TArray<FModuleDescriptor> Modules;
+	/** The plugin descriptor */
+	FPluginDescriptor Descriptor;
 };
 
 
