Set the default selection in the plugin browser to be the first category, rather than showing an empty top-level category.

diff --git a/Engine/Plugins/Editor/PluginsEditor/Source/PluginsEditor/Private/SPluginCategories.cpp b/Engine/Plugins/Editor/PluginsEditor/Source/PluginsEditor/Private/SPluginCategories.cpp
--- a/Engine/Plugins/Editor/PluginsEditor/Source/PluginsEditor/Private/SPluginCategories.cpp
+++ b/Engine/Plugins/Editor/PluginsEditor/Source/PluginsEditor/Private/SPluginCategories.cpp
@@ -41,7 +41,11 @@ void SPluginCategories::Construct( const FArguments& Args, const TSharedRef< SPl
 	// Select the first item by default
 	if( RootPluginCategories.Num() > 0 )
 	{
-		PluginCategoryTreeView->SetSelection( RootPluginCategories[ 0 ] );
+		TArray<FPluginCategoryTreeItemPtr> SubCategories = RootPluginCategories[0]->GetSubCategories();
+		if(SubCategories.Num() > 0)
+		{
+			PluginCategoryTreeView->SetSelection( SubCategories[ 0 ] );
+		}
 	}
 
 	ChildSlot.AttachWidget( PluginCategoryTreeView.ToSharedRef() );
